Remove the Chaos Cloth self-collision visualization that is misleadingly showing the nv cloth self collision indices.

diff --git a/Engine/Plugins/Experimental/ChaosCloth/Source/Chaos/Private/ChaosCloth/ChaosClothingSimulation.cpp b/Engine/Plugins/Experimental/ChaosCloth/Source/Chaos/Private/ChaosCloth/ChaosClothingSimulation.cpp
--- a/Engine/Plugins/Experimental/ChaosCloth/Source/Chaos/Private/ChaosCloth/ChaosClothingSimulation.cpp
+++ b/Engine/Plugins/Experimental/ChaosCloth/Source/Chaos/Private/ChaosCloth/ChaosClothingSimulation.cpp
@@ -154,7 +154,7 @@ void ClothingSimulation::CreateActor(USkeletalMeshComponent* InOwnerComponent, U
 	//Evolution->SetCCD(true); // ryan!!!
 
 	UClothingAssetCommon* Asset = Cast<UClothingAssetCommon>(InAsset);
-	const UChaosClothConfig* const ChaosClothSimConfig = Cast<UChaosClothConfig>(Asset->ChaosClothSimConfig);
+	const UChaosClothConfig* const ChaosClothSimConfig = Cast<UChaosClothConfig>(Asset-ClothConfigs);
 	check(ChaosClothSimConfig);
 
 	ClothingSimulationContext Context;
@@ -1635,33 +1635,6 @@ void ClothingSimulation::DebugDrawMaxDistances(USkeletalMeshComponent* OwnerComp
 	}
 }
 
-void ClothingSimulation::DebugDrawSelfCollision(USkeletalMeshComponent* OwnerComponent, FPrimitiveDrawInterface* PDI) const
-{
-	const TPBDParticles<float, 3>& Particles = Evolution->Particles();
-	for (int32 i = 0; i < IndexToRangeMap.Num(); ++i)
-	{
-		if (const UClothingAssetCommon* const Asset = Assets[i])
-		{
-			const UChaosClothConfig* const ChaosClothSimConfig = Cast<UChaosClothConfig>(Asset->ChaosClothSimConfig);
-			if (ChaosClothSimConfig && ChaosClothSimConfig->bUseSelfCollisions)
-			{
-				const FTransform RootBoneTransform = OwnerComponent->GetComponentSpaceTransforms()[Asset->ReferenceBoneIndex];
-
-				const UClothLODDataBase* LodData = Asset->ClothLodData[0];
-				const UClothPhysicalMeshDataBase* PhysMesh = LodData->PhysicalMeshData;
-				const TArray<uint32>& SelfCollisionIndices = PhysMesh->SelfCollisionIndices;
-				for (int32 SelfColIdx = 0; SelfColIdx < SelfCollisionIndices.Num(); ++SelfColIdx)
-				{
-					const FVector ParticlePosition =
-						RootBoneTransform.TransformPosition(
-							Particles.X(PhysMesh->SelfCollisionIndices[SelfColIdx]));
-					DrawWireSphere(PDI, ParticlePosition, FColor::White, Evolution->GetSelfCollisionThickness(), 8, SDPG_World, 0.0f, 0.001f);
-				}
-			}
-		}
-	}
-}
-
 void ClothingSimulation::DebugDrawAnimDrive(USkeletalMeshComponent* OwnerComponent, FPrimitiveDrawInterface* PDI) const
 {
 	const TPBDParticles<float, 3>& Particles = Evolution->Particles();
diff --git a/Engine/Plugins/Experimental/ChaosCloth/Source/Chaos/Public/ChaosCloth/ChaosClothingSimulation.h b/Engine/Plugins/Experimental/ChaosCloth/Source/Chaos/Public/ChaosCloth/ChaosClothingSimulation.h
--- a/Engine/Plugins/Experimental/ChaosCloth/Source/Chaos/Public/ChaosCloth/ChaosClothingSimulation.h
+++ b/Engine/Plugins/Experimental/ChaosCloth/Source/Chaos/Public/ChaosCloth/ChaosClothingSimulation.h
@@ -53,7 +53,6 @@ namespace Chaos
 		CHAOSCLOTH_API void DebugDrawCollision(USkeletalMeshComponent* OwnerComponent, FPrimitiveDrawInterface* PDI) const;
 		CHAOSCLOTH_API void DebugDrawBackstops(USkeletalMeshComponent* OwnerComponent, FPrimitiveDrawInterface* PDI) const;
 		CHAOSCLOTH_API void DebugDrawMaxDistances(USkeletalMeshComponent* OwnerComponent, FPrimitiveDrawInterface* PDI) const;
-		CHAOSCLOTH_API void DebugDrawSelfCollision(USkeletalMeshComponent* OwnerComponent, FPrimitiveDrawInterface* PDI) const;
 		CHAOSCLOTH_API void DebugDrawAnimDrive(USkeletalMeshComponent* OwnerComponent, FPrimitiveDrawInterface* PDI) const;
 #endif  // #if WITH_EDITOR
 
diff --git a/Engine/Plugins/Experimental/ChaosClothEditor/Source/Chaos/Private/ChaosClothEditor/ChaosSimulationEditorExtender.cpp b/Engine/Plugins/Experimental/ChaosClothEditor/Source/Chaos/Private/ChaosClothEditor/ChaosSimulationEditorExtender.cpp
--- a/Engine/Plugins/Experimental/ChaosClothEditor/Source/Chaos/Private/ChaosClothEditor/ChaosSimulationEditorExtender.cpp
+++ b/Engine/Plugins/Experimental/ChaosClothEditor/Source/Chaos/Private/ChaosClothEditor/ChaosSimulationEditorExtender.cpp
@@ -52,7 +52,6 @@ const FVisualizationOption FVisualizationOption::OptionData[] =
 	FVisualizationOption(&Chaos::ClothingSimulation::DebugDrawFaceNormals         , LOCTEXT("ChaosVisName_FaceNormals"         , "Face Normals"               ), LOCTEXT("ChaosVisName_FaceNormals_ToolTip"         , "Draws the current face normals for the simulation mesh")),
 	FVisualizationOption(&Chaos::ClothingSimulation::DebugDrawInversedFaceNormals , LOCTEXT("ChaosVisName_InversedFaceNormals" , "Inversed Face Normals"      ), LOCTEXT("ChaosVisName_InversedFaceNormals_ToolTip" , "Draws the inversed face normals for the simulation mesh")),
 	FVisualizationOption(&Chaos::ClothingSimulation::DebugDrawCollision           , LOCTEXT("ChaosVisName_Collision"           , "Collisions"                 ), LOCTEXT("ChaosVisName_Collision_ToolTip"           , "Draws the collision bodies the simulation is currently using")),
-	FVisualizationOption(&Chaos::ClothingSimulation::DebugDrawSelfCollision       , LOCTEXT("ChaosVisName_SelfCollision"       , "Self Collision Radii"       ), LOCTEXT("ChaosVisName_SelfCollision_ToolTip"       , "Draws the self collision radius for each particle if self collision is enabled")),
 	FVisualizationOption(&Chaos::ClothingSimulation::DebugDrawBackstops           , LOCTEXT("ChaosVisName_Backstop"            , "Backstops"                  ), LOCTEXT("ChaosVisName_Backstop_ToolTip"            , "Draws the backstop offset for each simulation particle"), /*bDisablesSimulation =*/true),
 	FVisualizationOption(&Chaos::ClothingSimulation::DebugDrawMaxDistances        , LOCTEXT("ChaosVisName_MaxDistance"         , "Max Distances"              ), LOCTEXT("ChaosVisName_MaxDistance_ToolTip"         , "Draws the current max distances for the sim particles as a line along its normal"), true),
 	FVisualizationOption(&Chaos::ClothingSimulation::DebugDrawAnimDrive           , LOCTEXT("ChaosVisName_AnimDrive"           , "Anim Drive"                 ), LOCTEXT("ChaosVisName_AnimDrive_Tooltip"           , "Draws the current skinned reference mesh for the simulation which anim drive will attempt to reach if enabled"))
