Enable Chaos Cloth plugins by default.

diff --git a/Engine/Plugins/Experimental/ChaosCloth/ChaosCloth.uplugin b/Engine/Plugins/Experimental/ChaosCloth/ChaosCloth.uplugin
--- a/Engine/Plugins/Experimental/ChaosCloth/ChaosCloth.uplugin
+++ b/Engine/Plugins/Experimental/ChaosCloth/ChaosCloth.uplugin
@@ -10,7 +10,7 @@
     "DocsURL" : "",
     "MarketplaceURL" : "",
     "SupportURL" : "",
-    "EnabledByDefault" : false,
+    "EnabledByDefault" : true,
     "CanContainContent" : false,
     "IsBetaVersion" : true,
     "Installed" : false,
diff --git a/Engine/Plugins/Experimental/ChaosClothEditor/ChaosClothEditor.uplugin b/Engine/Plugins/Experimental/ChaosClothEditor/ChaosClothEditor.uplugin
--- a/Engine/Plugins/Experimental/ChaosClothEditor/ChaosClothEditor.uplugin
+++ b/Engine/Plugins/Experimental/ChaosClothEditor/ChaosClothEditor.uplugin
@@ -10,7 +10,7 @@
 	"DocsURL": "",
 	"MarketplaceURL": "",
 	"SupportURL": "",
-	"EnabledByDefault": false,
+	"EnabledByDefault": true,
 	"CanContainContent": true,
 	"IsBetaVersion": true,
 	"Installed": false,
