Disable Chaos Cloth to avoid crash in tests

diff --git a/Engine/Plugins/Experimental/ChaosCloth/ChaosCloth.uplugin b/Engine/Plugins/Experimental/ChaosCloth/ChaosCloth.uplugin
--- a/Engine/Plugins/Experimental/ChaosCloth/ChaosCloth.uplugin
+++ b/Engine/Plugins/Experimental/ChaosCloth/ChaosCloth.uplugin
@@ -10,7 +10,7 @@
     "DocsURL" : "",
     "MarketplaceURL" : "",
     "SupportURL" : "",
-    "EnabledByDefault" : true,
+    "EnabledByDefault" : false,
     "CanContainContent" : false,
     "IsBetaVersion" : true,
     "Installed" : false,
diff --git a/Engine/Plugins/Experimental/ChaosClothEditor/ChaosClothEditor.uplugin b/Engine/Plugins/Experimental/ChaosClothEditor/ChaosClothEditor.uplugin
--- a/Engine/Plugins/Experimental/ChaosClothEditor/ChaosClothEditor.uplugin
+++ b/Engine/Plugins/Experimental/ChaosClothEditor/ChaosClothEditor.uplugin
@@ -10,7 +10,7 @@
 	"DocsURL": "",
 	"MarketplaceURL": "",
 	"SupportURL": "",
-	"EnabledByDefault": true,
+	"EnabledByDefault": false,
 	"CanContainContent": true,
 	"IsBetaVersion": true,
 	"Installed": false,
