Cloth - Add cloth editor text strings debug draw functionality.
Chaos Cloth - Add Particle Index debug draw.
Chaos Cloth - Add Maximum Distance Values debug draw.

diff --git a/Engine/Plugins/Experimental/ChaosCloth/Source/Chaos/Private/ChaosCloth/ChaosClothingSimulation.cpp b/Engine/Plugins/Experimental/ChaosCloth/Source/Chaos/Private/ChaosCloth/ChaosClothingSimulation.cpp
--- a/Engine/Plugins/Experimental/ChaosCloth/Source/Chaos/Private/ChaosCloth/ChaosClothingSimulation.cpp
+++ b/Engine/Plugins/Experimental/ChaosCloth/Source/Chaos/Private/ChaosCloth/ChaosClothingSimulation.cpp
@@ -47,6 +47,12 @@
 #include "PhysXIncludes.h"
 #endif
 
+#if WITH_EDITOR
+#include "Engine/Canvas.h"  // For debug draw text
+#include "CanvasItem.h"     //
+#include "Engine/Engine.h"  //
+#endif
+
 #if CHAOS_DEBUG_DRAW
 #include "Chaos/DebugDrawQueue.h"
 #include "HAL/IConsoleManager.h"
@@ -1895,10 +1901,86 @@ void ClothingSimulation::DebugDrawPhysMeshShaded(USkeletalMeshComponent* OwnerCo
 	LocalSimSpaceToWorld.SetOrigin(LocalSpaceLocation);
 	MeshBuilder.Draw(PDI, LocalSimSpaceToWorld, DebugClothMaterial->GetRenderProxy(), SDPG_World, false, false);
 }
+
+static void DrawText(FCanvas* Canvas, const FSceneView* SceneView, const FVector& Pos, const FText& Text, const FLinearColor& Color)
+{
+	FVector2D PixelLocation;
+	if (SceneView->WorldToPixel(Pos, PixelLocation))
+	{
+		FCanvasTextItem TextItem(PixelLocation, Text, GEngine->GetSmallFont(), Color);
+		TextItem.Scale = FVector2D::UnitVector;
+		TextItem.EnableShadow(FLinearColor::Black);
+		TextItem.Draw(Canvas);
+	}
+}
+
+void ClothingSimulation::DebugDrawParticleIndices(USkeletalMeshComponent* /*OwnerComponent*/, FCanvas* Canvas, const FSceneView* SceneView) const
+{
+	const TPBDParticles<float, 3>& Particles = Evolution->Particles();
+
+	static const FLinearColor DynamicColor = FColor::White;
+	static const FLinearColor KinematicColor = FColor::Purple;
+
+	for (int32 Index = 0; Index < Assets.Num(); ++Index)
+	{
+		const UClothingAssetCommon* const Asset = Assets[Index];
+		if (!Asset) { continue; }
+
+		for (uint32 ParticleIndex = IndexToRangeMap[Index][0]; ParticleIndex < IndexToRangeMap[Index][1]; ++ParticleIndex)
+		{
+			const FVector Position = Particles.X(ParticleIndex) + LocalSpaceLocation;
+
+			const FText Text = FText::AsNumber(ParticleIndex);
+			DrawText(Canvas, SceneView, Position, Text, Particles.InvM(ParticleIndex) == 0.0f ? KinematicColor : DynamicColor);
+		}
+	}
+}
+
+void ClothingSimulation::DebugDrawMaxDistanceValues(USkeletalMeshComponent* /*OwnerComponent*/, FCanvas* Canvas, const FSceneView* SceneView) const
+{
+	const TPBDParticles<float, 3>& Particles = Evolution->Particles();
+
+	static const FLinearColor DynamicColor = FColor::White;
+	static const FLinearColor KinematicColor = FColor::Purple;
+
+	FNumberFormattingOptions NumberFormattingOptions;
+	NumberFormattingOptions.AlwaysSign = false;
+	NumberFormattingOptions.UseGrouping = false;
+	NumberFormattingOptions.RoundingMode = ERoundingMode::HalfFromZero;
+	NumberFormattingOptions.MinimumIntegralDigits = 1;
+	NumberFormattingOptions.MaximumIntegralDigits = 6;
+	NumberFormattingOptions.MinimumFractionalDigits = 2;
+	NumberFormattingOptions.MaximumFractionalDigits = 2;
+
+	for (int32 Index = 0; Index < Assets.Num(); ++Index)
+	{
+		const UClothingAssetCommon* const Asset = Assets[Index];
+		if (!Asset) { continue; }
+
+		// Get Maximum Distances
+		const FClothLODDataCommon& AssetLodData = Asset->LodData[0];
+		const FClothPhysicalMeshData& PhysMesh = AssetLodData.PhysicalMeshData;
+		const FPointWeightMap& MaxDistances = PhysMesh.GetWeightMap(EChaosWeightMapTarget::MaxDistance);
+		if (MaxDistances.Num() == 0)
+		{
+			continue;
+		}
+		
+		for (uint32 ParticleIndex = IndexToRangeMap[Index][0]; ParticleIndex < IndexToRangeMap[Index][1]; ++ParticleIndex)
+		{
+			const uint32 WeightMapIndex = ParticleIndex - IndexToRangeMap[Index][0];
+			const float Distance = MaxDistances[WeightMapIndex];
+			const FVector Position = Particles.X(ParticleIndex) + LocalSpaceLocation;
+
+			const FText Text = FText::AsNumber(Distance, &NumberFormattingOptions);
+			DrawText(Canvas, SceneView, Position, Text, Particles.InvM(ParticleIndex) == 0.0f ? KinematicColor : DynamicColor);
+		}
+	}
+}
 #endif  // #if WITH_EDITOR
 
 #if WITH_EDITOR || CHAOS_DEBUG_DRAW
-static void DrawPoint(FPrimitiveDrawInterface* PDI, const FVector& Pos, const FLinearColor& Color, UMaterial* DebugClothMaterialVertex)  // USe color or material
+static void DrawPoint(FPrimitiveDrawInterface* PDI, const FVector& Pos, const FLinearColor& Color, UMaterial* DebugClothMaterialVertex)  // Use color or material
 {
 #if CHAOS_DEBUG_DRAW
 	if (!PDI)
diff --git a/Engine/Plugins/Experimental/ChaosCloth/Source/Chaos/Public/ChaosCloth/ChaosClothingSimulation.h b/Engine/Plugins/Experimental/ChaosCloth/Source/Chaos/Public/ChaosCloth/ChaosClothingSimulation.h
--- a/Engine/Plugins/Experimental/ChaosCloth/Source/Chaos/Public/ChaosCloth/ChaosClothingSimulation.h
+++ b/Engine/Plugins/Experimental/ChaosCloth/Source/Chaos/Public/ChaosCloth/ChaosClothingSimulation.h
@@ -56,6 +56,8 @@ namespace Chaos
 
 		// Editor only debug draw function
 		CHAOSCLOTH_API void DebugDrawPhysMeshShaded(USkeletalMeshComponent* OwnerComponent, FPrimitiveDrawInterface* PDI) const;
+		CHAOSCLOTH_API void DebugDrawParticleIndices(USkeletalMeshComponent* OwnerComponent, FCanvas* Canvas, const FSceneView* SceneView) const;
+		CHAOSCLOTH_API void DebugDrawMaxDistanceValues(USkeletalMeshComponent* OwnerComponent, FCanvas* Canvas, const FSceneView* SceneView) const;
 #endif  // #if WITH_EDITOR
 
 #if WITH_EDITOR || CHAOS_DEBUG_DRAW
diff --git a/Engine/Plugins/Experimental/ChaosClothEditor/Source/Chaos/Private/ChaosClothEditor/ChaosSimulationEditorExtender.cpp b/Engine/Plugins/Experimental/ChaosClothEditor/Source/Chaos/Private/ChaosClothEditor/ChaosSimulationEditorExtender.cpp
--- a/Engine/Plugins/Experimental/ChaosClothEditor/Source/Chaos/Private/ChaosClothEditor/ChaosSimulationEditorExtender.cpp
+++ b/Engine/Plugins/Experimental/ChaosClothEditor/Source/Chaos/Private/ChaosClothEditor/ChaosSimulationEditorExtender.cpp
@@ -24,7 +24,9 @@ namespace Chaos
 
 		// Chaos debug draw function
 		typedef void (Chaos::ClothingSimulation::*FDebugDrawFunction)(USkeletalMeshComponent*, FPrimitiveDrawInterface*) const;
+		typedef void (Chaos::ClothingSimulation::*FDebugDrawTextsFunction)(USkeletalMeshComponent*, FCanvas*, const FSceneView*) const;
 		FDebugDrawFunction DebugDrawFunction;
+		FDebugDrawTextsFunction DebugDrawTextsFunction;
 
 		FText DisplayName;         // Text for menu entries.
 		FText ToolTip;             // Text for menu tooltips.
@@ -38,6 +40,13 @@ namespace Chaos
 			, bDisablesSimulation(bInDisablesSimulation)
 			, bHidesClothSections(bInHidesClothSections)
 		{}
+		FVisualizationOption(FDebugDrawTextsFunction InDebugDrawTextsFunction, const FText& InDisplayName, const FText& InToolTip, bool bInDisablesSimulation = false, bool bInHidesClothSections = false)
+			: DebugDrawTextsFunction(InDebugDrawTextsFunction)
+			, DisplayName(InDisplayName)
+			, ToolTip(InToolTip)
+			, bDisablesSimulation(bInDisablesSimulation)
+			, bHidesClothSections(bInHidesClothSections)
+		{}
 	};
 }
 
@@ -47,13 +56,15 @@ const FVisualizationOption FVisualizationOption::OptionData[] =
 {
 	FVisualizationOption(&Chaos::ClothingSimulation::DebugDrawPhysMeshShaded      , LOCTEXT("ChaosVisName_PhysMesh"            , "Physical Mesh (Flat Shaded)"), LOCTEXT("ChaosVisName_PhysMeshShaded_ToolTip"      , "Draws the current physical result as a doubled sided flat shaded mesh"), /*bDisablesSimulation =*/false, /*bHidesClothSections=*/true),
 	FVisualizationOption(&Chaos::ClothingSimulation::DebugDrawPhysMeshWired       , LOCTEXT("ChaosVisName_PhysMeshWire"        , "Physical Mesh (Wireframe)"  ), LOCTEXT("ChaosVisName_PhysMeshWired_ToolTip"       , "Draws the current physical mesh result in wireframe")),
+	FVisualizationOption(&Chaos::ClothingSimulation::DebugDrawParticleIndices     , LOCTEXT("ChaosVisName_ParticleIndices"     , "Particle Indices"           ), LOCTEXT("ChaosVisName_ParticleIndices_ToolTip"     , "Draws the particle indices as instanciated by the solver")),
 	FVisualizationOption(&Chaos::ClothingSimulation::DebugDrawPointNormals        , LOCTEXT("ChaosVisName_PointNormals"        , "Point Normals"              ), LOCTEXT("ChaosVisName_PointNormals_ToolTip"        , "Draws the current point normals for the simulation mesh")),
 	FVisualizationOption(&Chaos::ClothingSimulation::DebugDrawInversedPointNormals, LOCTEXT("ChaosVisName_InversedPointNormals", "Inversed Point Normals"     ), LOCTEXT("ChaosVisName_InversedPointNormals_ToolTip", "Draws the inversed point normals for the simulation mesh")),
 	FVisualizationOption(&Chaos::ClothingSimulation::DebugDrawFaceNormals         , LOCTEXT("ChaosVisName_FaceNormals"         , "Face Normals"               ), LOCTEXT("ChaosVisName_FaceNormals_ToolTip"         , "Draws the current face normals for the simulation mesh")),
 	FVisualizationOption(&Chaos::ClothingSimulation::DebugDrawInversedFaceNormals , LOCTEXT("ChaosVisName_InversedFaceNormals" , "Inversed Face Normals"      ), LOCTEXT("ChaosVisName_InversedFaceNormals_ToolTip" , "Draws the inversed face normals for the simulation mesh")),
 	FVisualizationOption(&Chaos::ClothingSimulation::DebugDrawCollision           , LOCTEXT("ChaosVisName_Collision"           , "Collisions"                 ), LOCTEXT("ChaosVisName_Collision_ToolTip"           , "Draws the collision bodies the simulation is currently using")),
 	FVisualizationOption(&Chaos::ClothingSimulation::DebugDrawBackstops           , LOCTEXT("ChaosVisName_Backstop"            , "Backstops"                  ), LOCTEXT("ChaosVisName_Backstop_ToolTip"            , "Draws the backstop offset for each simulation particle"), /*bDisablesSimulation =*/true),
-	FVisualizationOption(&Chaos::ClothingSimulation::DebugDrawMaxDistances        , LOCTEXT("ChaosVisName_MaxDistance"         , "Max Distances"              ), LOCTEXT("ChaosVisName_MaxDistance_ToolTip"         , "Draws the current max distances for the sim particles as a line along its normal"), true),
+	FVisualizationOption(&Chaos::ClothingSimulation::DebugDrawMaxDistances        , LOCTEXT("ChaosVisName_MaxDistance"         , "Max Distances"              ), LOCTEXT("ChaosVisName_MaxDistance_ToolTip"         , "Draws the current max distances for the sim particles as a line along its normal"), /*bDisablesSimulation =*/true),
+	FVisualizationOption(&Chaos::ClothingSimulation::DebugDrawMaxDistanceValues   , LOCTEXT("ChaosVisName_MaxDistanceValue"    , "Max Distances As Numbers"   ), LOCTEXT("ChaosVisName_MaxDistanceValue_ToolTip"    , "Draws the current max distances as numbers")),
 	FVisualizationOption(&Chaos::ClothingSimulation::DebugDrawAnimDrive           , LOCTEXT("ChaosVisName_AnimDrive"           , "Anim Drive"                 ), LOCTEXT("ChaosVisName_AnimDrive_Tooltip"           , "Draws the current skinned reference mesh for the simulation which anim drive will attempt to reach if enabled")),
 	FVisualizationOption(&Chaos::ClothingSimulation::DebugDrawLongRangeConstraint , LOCTEXT("ChaosVisName_LongRangeConstraint" , "Long Range Constraint"      ), LOCTEXT("ChaosVisName_LongRangeConstraint_Tooltip" , "Draws the long range attachment constraint distances")),
 	FVisualizationOption(&Chaos::ClothingSimulation::DebugDrawWindDragForces      , LOCTEXT("ChaosVisName_WindDragForces"      , "Wind drag forces"           ), LOCTEXT("ChaosVisName_Wind_Tooltip"                , "Draws the Wind drag forces")),
@@ -122,13 +133,28 @@ void FSimulationEditorExtender::DebugDrawSimulation(const IClothingSimulation* S
 
 	for (int32 OptionIndex = 0; OptionIndex < FVisualizationOption::Count; ++OptionIndex)
 	{
-		if (Flags[OptionIndex])
+		if (Flags[OptionIndex] && FVisualizationOption::OptionData[OptionIndex].DebugDrawFunction)
 		{
 			(ChaosSimulation->*(FVisualizationOption::OptionData[OptionIndex].DebugDrawFunction))(OwnerComponent, PDI);
 		}
 	}
 }
 
+void FSimulationEditorExtender::DebugDrawSimulationTexts(const IClothingSimulation* Simulation, USkeletalMeshComponent* OwnerComponent, FCanvas* Canvas, const FSceneView* SceneView)
+{
+	if (!ensure(Simulation)) { return; }
+
+	const ClothingSimulation* const ChaosSimulation = static_cast<const ClothingSimulation*>(Simulation);
+
+	for (int32 OptionIndex = 0; OptionIndex < FVisualizationOption::Count; ++OptionIndex)
+	{
+		if (Flags[OptionIndex] && FVisualizationOption::OptionData[OptionIndex].DebugDrawTextsFunction)
+		{
+			(ChaosSimulation->*(FVisualizationOption::OptionData[OptionIndex].DebugDrawTextsFunction))(OwnerComponent, Canvas, SceneView);
+		}
+	}
+}
+
 bool FSimulationEditorExtender::ShouldDisableSimulation() const
 {
 	for (uint32 OptionIndex = 0; OptionIndex < FVisualizationOption::Count; ++OptionIndex)
diff --git a/Engine/Plugins/Experimental/ChaosClothEditor/Source/Chaos/Public/ChaosClothEditor/ChaosSimulationEditorExtender.h b/Engine/Plugins/Experimental/ChaosClothEditor/Source/Chaos/Public/ChaosClothEditor/ChaosSimulationEditorExtender.h
--- a/Engine/Plugins/Experimental/ChaosClothEditor/Source/Chaos/Public/ChaosClothEditor/ChaosSimulationEditorExtender.h
+++ b/Engine/Plugins/Experimental/ChaosClothEditor/Source/Chaos/Public/ChaosClothEditor/ChaosSimulationEditorExtender.h
@@ -7,9 +7,6 @@
 #include "SimulationEditorExtender.h"
 #include "Containers/BitArray.h"
 
-class USkeletalMeshComponent;
-class FPrimitiveDrawInterface;
-
 namespace Chaos
 {
 	/** Chaos extension to the asset editor. */
@@ -22,6 +19,7 @@ namespace Chaos
 		virtual UClass* GetSupportedSimulationFactoryClass() override;
 		virtual void ExtendViewportShowMenu(FMenuBuilder& MenuBuilder, TSharedRef<IPersonaPreviewScene> PreviewScene) override;
 		virtual void DebugDrawSimulation(const IClothingSimulation* Simulation, USkeletalMeshComponent* OwnerComponent, FPrimitiveDrawInterface* PDI) override;
+		virtual void DebugDrawSimulationTexts(const IClothingSimulation* Simulation, USkeletalMeshComponent* OwnerComponent, FCanvas* Canvas, const FSceneView* SceneView) override;
 		// End ISimulationEditorExtender Interface
 
 	private:
diff --git a/Engine/Source/Editor/ClothingSystemEditor/Public/SimulationEditorExtenderNv.h b/Engine/Source/Editor/ClothingSystemEditor/Public/SimulationEditorExtenderNv.h
--- a/Engine/Source/Editor/ClothingSystemEditor/Public/SimulationEditorExtenderNv.h
+++ b/Engine/Source/Editor/ClothingSystemEditor/Public/SimulationEditorExtenderNv.h
@@ -7,9 +7,6 @@
 #include "SimulationEditorExtender.h"
 #include "Internationalization/Text.h"
 
-class USkeletalMeshComponent;
-class FPrimitiveDrawInterface;
-
 struct FNvVisualizationOptions
 {
 	FNvVisualizationOptions();
@@ -73,6 +70,7 @@ public:
 	virtual UClass* GetSupportedSimulationFactoryClass() override;
 	virtual void ExtendViewportShowMenu(FMenuBuilder& InMenuBuilder, TSharedRef<IPersonaPreviewScene> InPreviewScene) override;
 	virtual void DebugDrawSimulation(const IClothingSimulation* InSimulation, USkeletalMeshComponent* InOwnerComponent, FPrimitiveDrawInterface* PDI) override;
+	virtual void DebugDrawSimulationTexts(const IClothingSimulation* InSimulation, USkeletalMeshComponent* InOwnerComponent, FCanvas* Canvas, const FSceneView* SceneView) override {}
 	// End ISimulationEditorExtender Interface
 
 private:
diff --git a/Engine/Source/Editor/ClothingSystemEditorInterface/Public/SimulationEditorExtender.h b/Engine/Source/Editor/ClothingSystemEditorInterface/Public/SimulationEditorExtender.h
--- a/Engine/Source/Editor/ClothingSystemEditorInterface/Public/SimulationEditorExtender.h
+++ b/Engine/Source/Editor/ClothingSystemEditorInterface/Public/SimulationEditorExtender.h
@@ -11,6 +11,8 @@ class IClothingSimulation;
 class USkeletalMeshComponent;
 class IPersonaPreviewScene;
 class FPrimitiveDrawInterface;
+class FCanvas;
+class FSceneView;
 
 /**
  * A simulation extender is an object responsible for extending certain editor features with respect to a certain
@@ -46,4 +48,14 @@ public:
 	 * @param PDI - The drawing interface to use
 	 */
 	virtual void DebugDrawSimulation(const IClothingSimulation* InSimulation, USkeletalMeshComponent* InOwnerComponent, FPrimitiveDrawInterface* PDI) = 0;
+
+	/**
+	 * Called from the editor when clothing is active to process any active debug drawing of any text strings
+	 * Recommended to use the show menu extension for controlling what data to draw
+	 * @param InSimulation - The running clothing simulation
+	 * @param InOwnerComponent - The component that owns the running clothing simulation
+	 * @param Canvas - The canvas to use for drawing the text
+	 * @param SceneView - The view on which to project the text
+	 */
+	virtual void DebugDrawSimulationTexts(const IClothingSimulation* InSimulation, USkeletalMeshComponent* InOwnerComponent, FCanvas* Canvas, const FSceneView* SceneView) = 0;
 };
diff --git a/Engine/Source/Editor/Persona/Private/AnimationEditorViewportClient.cpp b/Engine/Source/Editor/Persona/Private/AnimationEditorViewportClient.cpp
--- a/Engine/Source/Editor/Persona/Private/AnimationEditorViewportClient.cpp
+++ b/Engine/Source/Editor/Persona/Private/AnimationEditorViewportClient.cpp
@@ -482,6 +482,9 @@ void FAnimationViewportClient::DrawCanvas( FViewport& InViewport, FSceneView& Vi
 		{
 			DrawUVsForMesh(Viewport, &Canvas, 1.0f);
 		}
+
+		// Debug draw clothing texts
+		PreviewMeshComponent->DebugDrawClothingTexts(&Canvas, &View);
 	}
 }
 
diff --git a/Engine/Source/Runtime/Engine/Classes/Components/SkeletalMeshComponent.h b/Engine/Source/Runtime/Engine/Classes/Components/SkeletalMeshComponent.h
--- a/Engine/Source/Runtime/Engine/Classes/Components/SkeletalMeshComponent.h
+++ b/Engine/Source/Runtime/Engine/Classes/Components/SkeletalMeshComponent.h
@@ -31,6 +31,8 @@
 
 class Error;
 class FPrimitiveDrawInterface;
+class FCanvas;
+class FSceneView;
 class UAnimInstance;
 class UPhysicalMaterial;
 class USkeletalMeshComponent;
@@ -1417,6 +1419,13 @@ public:
 	 */
 	void DebugDrawClothing(FPrimitiveDrawInterface* PDI);
 
+	/**
+	 * Draw the currently clothing state, using the editor extender interface
+	 * @param Canvas The canvas to draw the text on
+	 * @param SceneView The view to project the text with
+	 */
+	void DebugDrawClothingTexts(FCanvas* Canvas, const FSceneView* SceneView);
+
 	/** Changes the value of bNotifyRigidBodyCollision
 	* @param bNewNotifyRigidBodyCollision - The value to assign to bNotifyRigidBodyCollision
 	*/
diff --git a/Engine/Source/Runtime/Engine/Private/SkeletalMeshComponentPhysics.cpp b/Engine/Source/Runtime/Engine/Private/SkeletalMeshComponentPhysics.cpp
--- a/Engine/Source/Runtime/Engine/Private/SkeletalMeshComponentPhysics.cpp
+++ b/Engine/Source/Runtime/Engine/Private/SkeletalMeshComponentPhysics.cpp
@@ -3292,6 +3292,25 @@ void USkeletalMeshComponent::DebugDrawClothing(FPrimitiveDrawInterface* PDI)
 #endif
 }
 
+void USkeletalMeshComponent::DebugDrawClothingTexts(FCanvas* Canvas, const FSceneView* SceneView)
+{
+#if WITH_EDITOR && ENABLE_DRAW_DEBUG
+
+	if(ClothingSimulation)
+	{
+		FClothingSystemEditorInterfaceModule& ClothingEditorModule = FModuleManager::LoadModuleChecked<FClothingSystemEditorInterfaceModule>(TEXT("ClothingSystemEditorInterface"));
+
+		if(ISimulationEditorExtender* Extender = ClothingEditorModule.GetSimulationEditorExtender(ClothingSimulationFactory->GetFName()))
+		{
+			Extender->DebugDrawSimulationTexts(ClothingSimulation, this, Canvas, SceneView);
+		}
+
+		return;
+	}
+
+#endif
+}
+
 void USkeletalMeshComponent::SetAllMassScale(float InMassScale)
 {
 	// Apply mass scale to each child body
