Chaos Cloth - Remove IsBetaVersion from the Chaos Cloth plugin, and update version name to 1.0.

diff --git a/Engine/Plugins/Experimental/ChaosCloth/ChaosCloth.uplugin b/Engine/Plugins/Experimental/ChaosCloth/ChaosCloth.uplugin
--- a/Engine/Plugins/Experimental/ChaosCloth/ChaosCloth.uplugin
+++ b/Engine/Plugins/Experimental/ChaosCloth/ChaosCloth.uplugin
@@ -1,7 +1,7 @@
 {
     "FileVersion" : 3,
     "Version" : 1,
-    "VersionName" : "0.1",
+    "VersionName" : "1.0",
     "FriendlyName" : "Chaos Cloth",
     "Description" : "Adds Chaos Cloth Module.",
     "Category": "Physics",
@@ -12,7 +12,6 @@
     "SupportURL" : "",
     "EnabledByDefault" : true,
     "CanContainContent" : false,
-    "IsBetaVersion" : true,
     "Installed" : false,
     "Modules": [
       {
diff --git a/Engine/Plugins/Experimental/ChaosClothEditor/ChaosClothEditor.uplugin b/Engine/Plugins/Experimental/ChaosClothEditor/ChaosClothEditor.uplugin
--- a/Engine/Plugins/Experimental/ChaosClothEditor/ChaosClothEditor.uplugin
+++ b/Engine/Plugins/Experimental/ChaosClothEditor/ChaosClothEditor.uplugin
@@ -1,7 +1,7 @@
 {
 	"FileVersion": 1,
 	"Version": 1,
-	"VersionName": "0.1",
+	"VersionName": "1.0",
 	"FriendlyName": "Chaos Cloth Editor",
 	"Description": "Editor module accompanying the Chaos Cloth runtime module.",
 	"Category": "Physics",
@@ -12,7 +12,6 @@
 	"SupportURL": "",
 	"EnabledByDefault": true,
 	"CanContainContent": true,
-	"IsBetaVersion": true,
 	"Installed": false,
 	"Modules": [
 		{
