Chaos Cloth - Updated the name of the old redundant Chaos Cloth Editor plugin to make it more obvious that it is deprecated.

#tests Opened the plug-in dialog and checked the name.
#rnx
#Synced-CL 49759196

diff --git a/Engine/Plugins/Experimental/ChaosClothEditor/ChaosClothEditor.uplugin b/Engine/Plugins/Experimental/ChaosClothEditor/ChaosClothEditor.uplugin
--- a/Engine/Plugins/Experimental/ChaosClothEditor/ChaosClothEditor.uplugin
+++ b/Engine/Plugins/Experimental/ChaosClothEditor/ChaosClothEditor.uplugin
@@ -2,7 +2,7 @@
 	"FileVersion": 1,
 	"Version": 1,
 	"VersionName": "1.0",
-	"FriendlyName": "Chaos Cloth Editor",
+	"FriendlyName": "Chaos Cloth Editor (Deprecated)",
 	"Description": "Deprecated plugin, please use the Chaos Cloth plugin instead.",
 	"Category": "Physics",
 	"CreatedBy": "Epic Games, Inc.",
